Script Editor Plugin has a dependency on Script Plugin. So they'll need the same blacklists

diff --git a/Engine/Plugins/Script/ScriptEditorPlugin/ScriptEditorPlugin.uplugin b/Engine/Plugins/Script/ScriptEditorPlugin/ScriptEditorPlugin.uplugin
--- a/Engine/Plugins/Script/ScriptEditorPlugin/ScriptEditorPlugin.uplugin
+++ b/Engine/Plugins/Script/ScriptEditorPlugin/ScriptEditorPlugin.uplugin
@@ -15,6 +15,7 @@
 		{
 			"Name" : "ScriptEditorPlugin",
 			"Type" : "Editor"
+			"BlacklistPlatforms" : [ "Linux" ]
 		}
 	]
 }
\ No newline at end of file
