Fixing a missing comma in .uplugin

diff --git a/Engine/Plugins/Script/ScriptEditorPlugin/ScriptEditorPlugin.uplugin b/Engine/Plugins/Script/ScriptEditorPlugin/ScriptEditorPlugin.uplugin
--- a/Engine/Plugins/Script/ScriptEditorPlugin/ScriptEditorPlugin.uplugin
+++ b/Engine/Plugins/Script/ScriptEditorPlugin/ScriptEditorPlugin.uplugin
@@ -14,7 +14,7 @@
 	[
 		{
 			"Name" : "ScriptEditorPlugin",
-			"Type" : "Editor"
+			"Type" : "Editor",
 			"BlacklistPlatforms" : [ "Linux" ]
 		}
 	]
