Script Plugin: adding the ability to create script classes from any Blueprintable base class

#change Removed ScriptActor, ScriptFunction <- no longer needed
#change Added Script Context subobjects which will be added to generated class as default subobjects
#change Added code to generate (UE) script proxy functions to call script (Lua etc) functions

diff --git a/Engine/Plugins/Script/ScriptEditorPlugin/Source/ScriptEditorPlugin/Private/ScriptBlueprintCompiler.cpp b/Engine/Plugins/Script/ScriptEditorPlugin/Source/ScriptEditorPlugin/Private/ScriptBlueprintCompiler.cpp
--- a/Engine/Plugins/Script/ScriptEditorPlugin/Source/ScriptEditorPlugin/Private/ScriptBlueprintCompiler.cpp
+++ b/Engine/Plugins/Script/ScriptEditorPlugin/Source/ScriptEditorPlugin/Private/ScriptBlueprintCompiler.cpp
@@ -6,12 +6,15 @@
 #include "ScriptBlueprintCompiler.h"
 #include "Kismet2NameValidators.h"
 #include "KismetReinstanceUtilities.h"
-#include "ScriptFunction.h"
+#include "ScriptContext.h"
+#include "ScriptContextComponent.h"
 
 ///-------------------------------------------------------------
 
 FScriptBlueprintCompiler::FScriptBlueprintCompiler(UScriptBlueprint* SourceSketch, FCompilerResultsLog& InMessageLog, const FKismetCompilerOptions& InCompilerOptions, TArray<UObject*>* InObjLoaded)
 	: Super(SourceSketch, InMessageLog, InCompilerOptions, InObjLoaded)
+	, NewScriptBlueprintClass(NULL)
+	, ContextProperty(NULL)
 {
 }
 
@@ -29,8 +32,9 @@ void FScriptBlueprintCompiler::CleanAndSanitizeClass(UBlueprintGeneratedClass* C
 	Super::CleanAndSanitizeClass(ClassToClean, OldCDO);
 
 	// Make sure our typed pointer is set
-	check(ClassToClean == NewClass);
+	check(ClassToClean == NewClass);	
 	NewScriptBlueprintClass = CastChecked<UScriptBlueprintGeneratedClass>((UObject*)NewClass);
+	ContextProperty = NULL;
 }
 
 void FScriptBlueprintCompiler::CreateClassVariablesFromBlueprint()
@@ -43,20 +47,8 @@ void FScriptBlueprintCompiler::CreateClassVariablesFromBlueprint()
 
 	for (auto& Field : ScriptDefinedFields)
 	{
-		if (Field.Class->IsChildOf(UFunction::StaticClass()))
-		{
-			// Functions are currently only supported for AScriptActors
-			if (Blueprint->ParentClass->IsChildOf(AScriptActor::StaticClass()))
-			{
-				UFunction* NewFunction = new(NewClass, Field.Name, RF_Public) UScriptFunction(FPostConstructInitializeProperties());
-				NewFunction->Bind();
-				NewFunction->StaticLink();
-				NewClass->AddFunctionToFunctionMap(NewFunction);
-				NewFunction->Next = NewClass->Children;
-				NewClass->Children = NewFunction;
-			}
-		}
-		else if (Field.Class->IsChildOf(UProperty::StaticClass()))
+		UClass* InnerType = Field.Class;
+		if (Field.Class->IsChildOf(UProperty::StaticClass()))
 		{
 			FString PinCategory;
 			if (Field.Class->IsChildOf(UStrProperty::StaticClass()))
@@ -78,20 +70,103 @@ void FScriptBlueprintCompiler::CreateClassVariablesFromBlueprint()
 			else if (Field.Class->IsChildOf(UObjectProperty::StaticClass()))
 			{
 				PinCategory = Schema->PC_Object;
+				// @todo: some scripting extensions (that are strongly typed) can handle this better
+				InnerType = UObject::StaticClass();
 			}
 			if (!PinCategory.IsEmpty())
-			{			
-				FEdGraphPinType ScriptPinType(PinCategory, TEXT(""), Field.Class, false, false);
+			{
+				FEdGraphPinType ScriptPinType(PinCategory, TEXT(""), InnerType, false, false);
 				UProperty* ScriptProperty = CreateVariable(Field.Name, ScriptPinType);
 				if (ScriptProperty != NULL)
 				{
 					ScriptProperty->SetMetaData(TEXT("Category"), *Blueprint->GetName());
-					ScriptProperty->SetPropertyFlags(CPF_BlueprintVisible | CPF_Edit);					
+					ScriptProperty->SetPropertyFlags(CPF_BlueprintVisible | CPF_Edit);
 					NewScripClass->ScriptProperties.Add(ScriptProperty);
 				}
 			}
 		}
 	}
+
+	CreateScriptContextProperty();
+}
+
+void FScriptBlueprintCompiler::CreateScriptContextProperty()
+{
+	// The only case we don't need a script context is if the script class derives form UScriptComponent
+	UClass* ContextClass = nullptr;
+	if (Blueprint->ParentClass->IsChildOf(AActor::StaticClass()))
+	{
+		ContextClass = UScriptContextComponent::StaticClass();
+	}
+	else if (!Blueprint->ParentClass->IsChildOf(UScriptComponent::StaticClass()))
+	{
+		ContextClass = UScriptContext::StaticClass();
+	}
+
+	if (ContextClass)
+	{
+		FEdGraphPinType ScriptContextPinType(Schema->PC_Object, TEXT(""), ContextClass, false, false);
+		ContextProperty = CastChecked<UObjectProperty>(CreateVariable(TEXT("Generated_ScriptContext"), ScriptContextPinType));
+		ContextProperty->SetPropertyFlags(CPF_ContainsInstancedReference | CPF_InstancedReference);
+	}
+}
+
+void FScriptBlueprintCompiler::CreateFunctionList()
+{
+	Super::CreateFunctionList();
+
+	if (!Blueprint->ParentClass->IsChildOf(UScriptComponent::StaticClass()))
+	{
+		for (auto& Field : ScriptDefinedFields)
+		{
+			if (Field.Class->IsChildOf(UFunction::StaticClass()))
+			{
+				CreateScriptDefinedFunction(Field);
+			}
+		}
+	}
+}
+
+void FScriptBlueprintCompiler::CreateScriptDefinedFunction(FScriptField& Field)
+{
+	check(ContextProperty);
+	
+	UScriptBlueprint* Blueprint = ScriptBlueprint();
+	const FString FunctionName = Field.Name.ToString();
+
+	// Create Blueprint Graph which consists of 3 nodes: 'Entry', 'Get Script Context' and 'Call Function'
+	// @todo: once we figure out how to get parameter lists for functions we can add suport for that here
+
+	UEdGraph* ScriptFunctionGraph = NewNamedObject<UEdGraph>(Blueprint, *FString::Printf(TEXT("%s_Graph"), *FunctionName));
+	ScriptFunctionGraph->Schema = UEdGraphSchema_K2::StaticClass();
+	ScriptFunctionGraph->SetFlags(RF_Transient);
+	
+	FKismetFunctionContext* FunctionContext = CreateFunctionContext();
+	FunctionContext->SourceGraph = ScriptFunctionGraph;
+	FunctionContext->bCreateDebugData = false;
+
+	UK2Node_FunctionEntry* EntryNode = SpawnIntermediateNode<UK2Node_FunctionEntry>(NULL, ScriptFunctionGraph);
+	EntryNode->CustomGeneratedFunctionName = Field.Name;
+	EntryNode->AllocateDefaultPins();
+
+	UK2Node_VariableGet* GetVariableNode = SpawnIntermediateNode<UK2Node_VariableGet>(NULL, ScriptFunctionGraph);
+	GetVariableNode->VariableReference.SetSelfMember(ContextProperty->GetFName());
+	GetVariableNode->AllocateDefaultPins();
+
+	
+	UK2Node_CallFunction* CallFunctionNode = SpawnIntermediateNode<UK2Node_CallFunction>(NULL, ScriptFunctionGraph);
+	CallFunctionNode->FunctionReference.SetExternalMember(TEXT("CallScriptFunction"), ContextProperty->PropertyClass);
+	CallFunctionNode->AllocateDefaultPins();
+	UEdGraphPin* FunctionNamePin = CallFunctionNode->FindPinChecked(TEXT("FunctionName"));
+	FunctionNamePin->DefaultValue = FunctionName;
+
+	// Link nodes together
+	UEdGraphPin* ExecPin = Schema->FindExecutionPin(*EntryNode, EGPD_Output);
+	UEdGraphPin* GetVariableOutPin = GetVariableNode->FindPinChecked(ContextProperty->GetName());
+	UEdGraphPin* CallFunctionPin = Schema->FindExecutionPin(*CallFunctionNode, EGPD_Input);
+	UEdGraphPin* FunctionTargetPin = CallFunctionNode->FindPinChecked(TEXT("self"));
+	ExecPin->MakeLinkTo(CallFunctionPin);
+	GetVariableOutPin->MakeLinkTo(FunctionTargetPin);
 }
 
 void FScriptBlueprintCompiler::FinishCompilingClass(UClass* Class)
@@ -102,33 +177,34 @@ void FScriptBlueprintCompiler::FinishCompilingClass(UClass* Class)
 	ScriptClass->SourceCode = Blueprint->SourceCode;
 	ScriptClass->ByteCode = Blueprint->ByteCode;
 
-	Super::FinishCompilingClass(Class);
-
 	// Allow Blueprint Components to be used in Blueprints
 	if (Blueprint->ParentClass->IsChildOf(UScriptComponent::StaticClass()) && Class != Blueprint->SkeletonGeneratedClass)
 	{
 		Class->SetMetaData(TEXT("BlueprintSpawnableComponent"), TEXT("true"));
 	}
+
+	Super::FinishCompilingClass(Class);
+
+	// Ff context property has been created, create a DSO and set it on the CDO
+	if (ContextProperty)
+	{
+		UObject* CDO = Class->GetDefaultObject();
+		UObject* ContextDefaultSubobject = StaticConstructObject(ContextProperty->PropertyClass, CDO, "ScriptContext", RF_DefaultSubObject);
+		ContextProperty->SetObjectPropertyValue(ContextProperty->ContainerPtrToValuePtr<UObject*>(CDO), ContextDefaultSubobject);
+	}
 }
 
 void FScriptBlueprintCompiler::Compile()
 {
-	ScriptContext = NewScriptBlueprintClass->CreateContext();
+	ScriptBlueprint()->UpdateSourceCodeIfChanged();
+	ScriptContext = FScriptContextBase::CreateContext(ScriptBlueprint()->SourceCode, NULL, NULL);
 	bool Result = true;
 	if (ScriptContext.IsValid())
-	{
-		ScriptBlueprint()->UpdateSourceCodeIfChanged();
-		Result = ScriptContext->Initialize(ScriptBlueprint()->SourceCode, NULL);
-		if (!Result)
-		{
-			ScriptContext.Reset();
-		}
-	}
-	if (Result)
 	{
 		ScriptDefinedFields.Empty();
 		ScriptContext->GetScriptDefinedFields(ScriptDefinedFields);
 	}
+	ContextProperty = NULL;
 
 	Super::Compile();
 }
diff --git a/Engine/Plugins/Script/ScriptEditorPlugin/Source/ScriptEditorPlugin/Private/ScriptBlueprintCompiler.h b/Engine/Plugins/Script/ScriptEditorPlugin/Source/ScriptEditorPlugin/Private/ScriptBlueprintCompiler.h
--- a/Engine/Plugins/Script/ScriptEditorPlugin/Source/ScriptEditorPlugin/Private/ScriptBlueprintCompiler.h
+++ b/Engine/Plugins/Script/ScriptEditorPlugin/Source/ScriptEditorPlugin/Private/ScriptBlueprintCompiler.h
@@ -27,6 +27,7 @@ protected:
 	virtual void CleanAndSanitizeClass(UBlueprintGeneratedClass* ClassToClean, UObject*& OldCDO) override;
 	virtual void EnsureProperGeneratedClass(UClass*& TargetClass) override;
 	virtual void CreateClassVariablesFromBlueprint() override;
+	virtual void CreateFunctionList() override;
 	virtual void FinishCompilingClass(UClass* Class) override;
 	virtual bool ValidateGeneratedClass(UBlueprintGeneratedClass* Class) override;
 	virtual void CopyTermDefaultsToDefaultObject(UObject* DefaultObject) override;
@@ -34,11 +35,25 @@ protected:
 
 protected:
 
+	/**
+	 * Creates a script context property for this class (if needed)
+	 */
+	void CreateScriptContextProperty();
+
+	/**
+	 * Creates a Blueprint Graph function definition for script defined function
+	 *
+	 * @param Field Function exported by script
+	 */
+	void CreateScriptDefinedFunction(FScriptField& Field);
+
 	/** New script class */
 	UScriptBlueprintGeneratedClass* NewScriptBlueprintClass;
 	/** Script context */
 	TAutoPtr<FScriptContextBase> ScriptContext;
 	/** Script-defined properties and functions */
 	TArray<FScriptField> ScriptDefinedFields;
+	/** Script context property generated for the compiled class */
+	UObjectProperty* ContextProperty;
 };
 
diff --git a/Engine/Plugins/Script/ScriptEditorPlugin/Source/ScriptEditorPlugin/Private/ScriptEditorPluginPrivatePCH.h b/Engine/Plugins/Script/ScriptEditorPlugin/Source/ScriptEditorPlugin/Private/ScriptEditorPluginPrivatePCH.h
--- a/Engine/Plugins/Script/ScriptEditorPlugin/Source/ScriptEditorPlugin/Private/ScriptEditorPluginPrivatePCH.h
+++ b/Engine/Plugins/Script/ScriptEditorPlugin/Source/ScriptEditorPlugin/Private/ScriptEditorPluginPrivatePCH.h
@@ -10,7 +10,6 @@
 #include "ScriptFactory.h"
 #include "ReimportScriptFactory.h"
 #include "AssetTypeActions_Base.h"
-#include "ScriptActor.h"
 #include "ScriptComponent.h"
 #include "ClassViewerModule.h"
 #include "ClassViewerFilter.h"
diff --git a/Engine/Plugins/Script/ScriptEditorPlugin/Source/ScriptEditorPlugin/Private/ScriptFactory.cpp b/Engine/Plugins/Script/ScriptEditorPlugin/Source/ScriptEditorPlugin/Private/ScriptFactory.cpp
--- a/Engine/Plugins/Script/ScriptEditorPlugin/Source/ScriptEditorPlugin/Private/ScriptFactory.cpp
+++ b/Engine/Plugins/Script/ScriptEditorPlugin/Source/ScriptEditorPlugin/Private/ScriptFactory.cpp
@@ -1,33 +1,11 @@
 // Copyright 1998-2014 Epic Games, Inc. All Rights Reserved. 
 #include "ScriptEditorPluginPrivatePCH.h"
 
-/**
- * Class Picker filter
- */
-class FScriptFactoryFilter : public IClassViewerFilter
-{
-public:
-	TSet< const UClass* > AllowedChildOfClasses;
-
-	bool IsClassAllowed(const FClassViewerInitializationOptions& InInitOptions, const UClass* InClass, TSharedRef< FClassViewerFilterFuncs > InFilterFuncs)
-	{
-		return AllowedChildOfClasses.Contains(InClass);
-	}
-
-	virtual bool IsUnloadedClassAllowed(const FClassViewerInitializationOptions& InInitOptions, const TSharedRef< const IUnloadedBlueprintData > InUnloadedClassData, TSharedRef< FClassViewerFilterFuncs > InFilterFuncs) override
-	{
-		return true;
-	}
-};
-
-//////////////////////////////////////////////////////////////////////////
-
-
 UScriptFactory::UScriptFactory(const FPostConstructInitializeProperties& PCIP)
 	: Super( PCIP )
 {
 	SupportedClass = UScriptBlueprint::StaticClass();
-	ParentClass = AScriptActor::StaticClass();
+	ParentClass = AActor::StaticClass();
 
 	Formats.Add(TEXT("txt;Script"));
 	Formats.Add(TEXT("lua;Script"));
@@ -54,16 +32,10 @@ bool UScriptFactory::ConfigureProperties()
 	// Fill in options
 	FClassViewerInitializationOptions Options;
 	Options.Mode = EClassViewerMode::ClassPicker;
-	Options.DisplayMode = EClassViewerDisplayMode::ListView;
-	Options.bShowObjectRootClass = false;
-	Options.bIsBlueprintBaseOnly = false;
-	Options.bShowUnloadedBlueprints = false;
-
-	// Only allow script-enabled classes
-	TSharedPtr<FScriptFactoryFilter> Filter = MakeShareable(new FScriptFactoryFilter);
-	Options.ClassFilter = Filter;
-	Filter->AllowedChildOfClasses.Add(AScriptActor::StaticClass());
-	Filter->AllowedChildOfClasses.Add(UScriptComponent::StaticClass());
+	Options.DisplayMode = EClassViewerDisplayMode::TreeView;
+	Options.bShowObjectRootClass = true;
+	Options.bIsBlueprintBaseOnly = true;
+	Options.bShowUnloadedBlueprints = true;
 
 	const FText TitleText = NSLOCTEXT("EditorFactories", "CreateScriptOptions", "Pick Parent Class");
 	UClass* ChosenClass = NULL;
diff --git a/Engine/Plugins/Script/ScriptEditorPlugin/Source/ScriptEditorPlugin/ScriptEditorPlugin.Build.cs b/Engine/Plugins/Script/ScriptEditorPlugin/Source/ScriptEditorPlugin/ScriptEditorPlugin.Build.cs
--- a/Engine/Plugins/Script/ScriptEditorPlugin/Source/ScriptEditorPlugin/ScriptEditorPlugin.Build.cs
+++ b/Engine/Plugins/Script/ScriptEditorPlugin/Source/ScriptEditorPlugin/ScriptEditorPlugin.Build.cs
@@ -31,7 +31,8 @@ namespace UnrealBuildTool.Rules
 					"ScriptPlugin",
 					"ClassViewer",
 					"KismetCompiler",
-					"Kismet"
+					"Kismet",
+					"BlueprintGraph"
 					// ... add other public dependencies that you statically link with here ...
 				}
 				);
diff --git a/Engine/Plugins/Script/ScriptPlugin/Source/ScriptPlugin/Classes/ScriptActor.h b/Engine/Plugins/Script/ScriptPlugin/Source/ScriptPlugin/Classes/ScriptActor.h
deleted file mode 100644
diff --git a/Engine/Plugins/Script/ScriptPlugin/Source/ScriptPlugin/Classes/ScriptBlueprint.h b/Engine/Plugins/Script/ScriptPlugin/Source/ScriptPlugin/Classes/ScriptBlueprint.h
--- a/Engine/Plugins/Script/ScriptPlugin/Source/ScriptPlugin/Classes/ScriptBlueprint.h
+++ b/Engine/Plugins/Script/ScriptPlugin/Source/ScriptPlugin/Classes/ScriptBlueprint.h
@@ -7,7 +7,7 @@
 #include "ScriptBlueprint.generated.h"
 
 /**
- * The Script blueprint enables extending AScriptActor and UScriptComponent classes
+ * The Script blueprint generates script-defined classes
  */
 UCLASS(BlueprintType, EarlyAccessPreview)
 class SCRIPTPLUGIN_API UScriptBlueprint : public UBlueprint
diff --git a/Engine/Plugins/Script/ScriptPlugin/Source/ScriptPlugin/Classes/ScriptBlueprintGeneratedClass.h b/Engine/Plugins/Script/ScriptPlugin/Source/ScriptPlugin/Classes/ScriptBlueprintGeneratedClass.h
--- a/Engine/Plugins/Script/ScriptPlugin/Source/ScriptPlugin/Classes/ScriptBlueprintGeneratedClass.h
+++ b/Engine/Plugins/Script/ScriptPlugin/Source/ScriptPlugin/Classes/ScriptBlueprintGeneratedClass.h
@@ -10,7 +10,7 @@ class UScriptBlueprintGeneratedClass;
 /**
 * Script-defined field (variable or function)
 */
-struct FScriptField
+struct SCRIPTPLUGIN_API FScriptField
 {
 	/** Field name */
 	FName Name;
@@ -29,11 +29,17 @@ struct FScriptField
 };
 
 /** Script context for this component */
-class FScriptContextBase
+class SCRIPTPLUGIN_API FScriptContextBase
 {
 public:
+	
 	virtual ~FScriptContextBase() {}
 
+	/**
+	* Creates a script context object
+	*/
+	static FScriptContextBase* CreateContext(const FString& SourceCode, UScriptBlueprintGeneratedClass* Class, UObject* Owner);
+
 	/**
 	 * Initializes script context given script code
 	 * @param Code Script code
@@ -129,11 +135,6 @@ class SCRIPTPLUGIN_API UScriptBlueprintGeneratedClass : public UBlueprintGenerat
 	virtual void Link(FArchive& Ar, bool bRelinkExistingProperties) override;
 	virtual void PurgeClass(bool bRecompilingOnLoad) override;
 
-	/**
-	 * Creates a script context object
-	 */
-	FScriptContextBase* CreateContext();
-
 	/**
 	* Adds a unique native function mapping
 	* @param InName Name of the native function
diff --git a/Engine/Plugins/Script/ScriptPlugin/Source/ScriptPlugin/Classes/ScriptContext.h b/Engine/Plugins/Script/ScriptPlugin/Source/ScriptPlugin/Classes/ScriptContext.h
new file mode 100644
--- /dev/null
+++ b/Engine/Plugins/Script/ScriptPlugin/Source/ScriptPlugin/Classes/ScriptContext.h
@@ -0,0 +1,25 @@
+// Copyright 1998-2014 Epic Games, Inc. All Rights Reserved.
+#pragma once
+
+#include "Engine/BlueprintGeneratedClass.h"
+#include "ScriptContext.generated.h"
+
+/**
+* Script context subobject
+*/
+UCLASS(DefaultToInstanced, EarlyAccessPreview)
+class SCRIPTPLUGIN_API UScriptContext : public UObject
+{
+	GENERATED_UCLASS_BODY()
+
+	virtual void PostInitProperties() override;
+	virtual void BeginDestroy() override;
+
+	UFUNCTION(BlueprintCallable, Category="ScriptContext")
+	void CallScriptFunction(FString FunctionName);
+
+protected:
+
+	/** Script context (code) */
+	FScriptContextBase* Context;
+};
\ No newline at end of file
diff --git a/Engine/Plugins/Script/ScriptPlugin/Source/ScriptPlugin/Classes/ScriptContextComponent.h b/Engine/Plugins/Script/ScriptPlugin/Source/ScriptPlugin/Classes/ScriptContextComponent.h
new file mode 100644
--- /dev/null
+++ b/Engine/Plugins/Script/ScriptPlugin/Source/ScriptPlugin/Classes/ScriptContextComponent.h
@@ -0,0 +1,36 @@
+// Copyright 1998-2014 Epic Games, Inc. All Rights Reserved.
+
+#pragma once
+#include "ScriptBlueprintGeneratedClass.h"
+#include "ScriptContextComponent.generated.h"
+
+/** 
+ * Script-extendable component class
+ */
+UCLASS(EarlyAccessPreview)
+class SCRIPTPLUGIN_API UScriptContextComponent : public UActorComponent
+{
+	GENERATED_UCLASS_BODY()
+
+public:
+
+	// Begin UActorComponent interface.
+	virtual void OnRegister() override;
+	virtual void InitializeComponent() override;
+	virtual void TickComponent(float DeltaTime, enum ELevelTick TickType, FActorComponentTickFunction *ThisTickFunction) override;
+	virtual void BeginDestroy() override;
+	// Begin UActorComponent interface.
+
+	/**
+	* Calls a script-defined function (no arguments)
+	* @param FunctionName Name of the function to call
+	*/
+	UFUNCTION(BlueprintCallable, Category = "ScriptContext")
+	void CallScriptFunction(FString FunctionName);
+
+protected:
+
+	/** Script context (code) */
+	FScriptContextBase* Context;
+};
+
diff --git a/Engine/Plugins/Script/ScriptPlugin/Source/ScriptPlugin/Classes/ScriptFunction.h b/Engine/Plugins/Script/ScriptPlugin/Source/ScriptPlugin/Classes/ScriptFunction.h
deleted file mode 100644
diff --git a/Engine/Plugins/Script/ScriptPlugin/Source/ScriptPlugin/Private/LuaIntegration.cpp b/Engine/Plugins/Script/ScriptPlugin/Source/ScriptPlugin/Private/LuaIntegration.cpp
--- a/Engine/Plugins/Script/ScriptPlugin/Source/ScriptPlugin/Private/LuaIntegration.cpp
+++ b/Engine/Plugins/Script/ScriptPlugin/Source/ScriptPlugin/Private/LuaIntegration.cpp
@@ -3,7 +3,6 @@
 #include "ScriptPluginPrivatePCH.h"
 #include "LuaIntegration.h"
 #include "ScriptObjectReferencer.h"
-#include "ScriptActor.h"
 
 #if WITH_LUA
 
diff --git a/Engine/Plugins/Script/ScriptPlugin/Source/ScriptPlugin/Private/ScriptActor.cpp b/Engine/Plugins/Script/ScriptPlugin/Source/ScriptPlugin/Private/ScriptActor.cpp
deleted file mode 100644
diff --git a/Engine/Plugins/Script/ScriptPlugin/Source/ScriptPlugin/Private/ScriptBlueprintGeneratedClass.cpp b/Engine/Plugins/Script/ScriptPlugin/Source/ScriptPlugin/Private/ScriptBlueprintGeneratedClass.cpp
--- a/Engine/Plugins/Script/ScriptPlugin/Source/ScriptPlugin/Private/ScriptBlueprintGeneratedClass.cpp
+++ b/Engine/Plugins/Script/ScriptPlugin/Source/ScriptPlugin/Private/ScriptBlueprintGeneratedClass.cpp
@@ -23,15 +23,6 @@ void UScriptBlueprintGeneratedClass::Link(FArchive& Ar, bool bRelinkExistingProp
 	Super::Link(Ar, bRelinkExistingProperties);
 }
 
-FScriptContextBase* UScriptBlueprintGeneratedClass::CreateContext()
-{
-	FScriptContextBase* NewContext = NULL;
-#if WITH_LUA
-	NewContext = new FLuaContext();
-#endif
-	return NewContext;
-}
-
 void UScriptBlueprintGeneratedClass::AddUniqueNativeFunction(const FName& InName, Native InPointer)
 {
 	// Find the function in the class's native function lookup table.
@@ -130,4 +121,29 @@ void FScriptContextBase::FetchScriptPropertyValues(UScriptBlueprintGeneratedClas
 			StringProperty->SetPropertyValue(Property->ContainerPtrToValuePtr<FString>(Obj), Value);
 		}
 	}
+}
+
+FScriptContextBase* FScriptContextBase::CreateContext(const FString& SourceCode, UScriptBlueprintGeneratedClass* Class, UObject* Owner)
+{
+	FScriptContextBase* NewContext = NULL;
+#if WITH_LUA
+	NewContext = new FLuaContext();
+#endif
+	if (NewContext)
+	{
+		if (NewContext->Initialize(SourceCode, Owner))
+		{
+			// Push values set by CDO
+			if (Class && Owner)
+			{
+				NewContext->PushScriptPropertyValues(Class, Owner);
+			}
+		}
+		else
+		{
+			delete NewContext;
+			NewContext = NULL;
+		}
+	}
+	return NewContext;
 }
\ No newline at end of file
diff --git a/Engine/Plugins/Script/ScriptPlugin/Source/ScriptPlugin/Private/ScriptComponent.cpp b/Engine/Plugins/Script/ScriptPlugin/Source/ScriptPlugin/Private/ScriptComponent.cpp
--- a/Engine/Plugins/Script/ScriptPlugin/Source/ScriptPlugin/Private/ScriptComponent.cpp
+++ b/Engine/Plugins/Script/ScriptPlugin/Source/ScriptPlugin/Private/ScriptComponent.cpp
@@ -21,22 +21,11 @@ void UScriptComponent::OnRegister()
 	auto ScriptClass = UScriptBlueprintGeneratedClass::GetScriptGeneratedClass(GetClass());
 	if (ScriptClass && GetWorld() && GetWorld()->WorldType != EWorldType::Editor)
 	{
-		Context = ScriptClass->CreateContext();
-		if (Context)
+		Context = FScriptContextBase::CreateContext(ScriptClass->SourceCode, ScriptClass, this);
+		if (!Context || !Context->CanTick())
 		{
-			bool bDoNotTick = true;
-			if (Context->Initialize(ScriptClass->SourceCode, this))
-			{
-				bDoNotTick = !Context->CanTick();
-				// Push values set by CDO
-				Context->PushScriptPropertyValues(ScriptClass, this);
-			}
-
-			if (bDoNotTick)
-			{
-				bAutoActivate = false;
-				PrimaryComponentTick.bCanEverTick = false;
-			}
+			bAutoActivate = false;
+			PrimaryComponentTick.bCanEverTick = false;
 		}
 	}
 }
diff --git a/Engine/Plugins/Script/ScriptPlugin/Source/ScriptPlugin/Private/ScriptContext.cpp b/Engine/Plugins/Script/ScriptPlugin/Source/ScriptPlugin/Private/ScriptContext.cpp
new file mode 100644
--- /dev/null
+++ b/Engine/Plugins/Script/ScriptPlugin/Source/ScriptPlugin/Private/ScriptContext.cpp
@@ -0,0 +1,52 @@
+// Copyright 1998-2014 Epic Games, Inc. All Rights Reserved. 
+#include "ScriptPluginPrivatePCH.h"
+
+//////////////////////////////////////////////////////////////////////////
+
+UScriptContext::UScriptContext(const FPostConstructInitializeProperties& PCIP)
+	: Super(PCIP)
+{
+	Context = NULL;
+}
+
+void UScriptContext::PostInitProperties()
+{
+	Super::PostInitProperties();
+
+	if (!HasAnyFlags(RF_ClassDefaultObject))
+	{
+		auto ContextOwner = GetOuter();
+		if (ContextOwner && !ContextOwner->HasAnyFlags(RF_ClassDefaultObject))
+		{
+			auto ScriptClass = UScriptBlueprintGeneratedClass::GetScriptGeneratedClass(ContextOwner->GetClass());
+			if (ScriptClass)
+			{
+				Context = FScriptContextBase::CreateContext(ScriptClass->SourceCode, ScriptClass, ContextOwner);
+			}
+		}
+	}
+}
+
+void UScriptContext::BeginDestroy()
+{
+	if (Context)
+	{
+		Context->Destroy();
+		delete Context;
+		Context = NULL;
+	}
+
+	Super::BeginDestroy();
+}
+
+void UScriptContext::CallScriptFunction(FString FunctionName)
+{
+	if (Context)
+	{
+		auto ContextOwner = GetOuter();
+		auto ScriptClass = UScriptBlueprintGeneratedClass::GetScriptGeneratedClass(ContextOwner->GetClass());
+		Context->PushScriptPropertyValues(ScriptClass, ContextOwner);
+		Context->CallFunction(FunctionName);
+		Context->FetchScriptPropertyValues(ScriptClass, ContextOwner);
+	}
+}
\ No newline at end of file
diff --git a/Engine/Plugins/Script/ScriptPlugin/Source/ScriptPlugin/Private/ScriptContextComponent.cpp b/Engine/Plugins/Script/ScriptPlugin/Source/ScriptPlugin/Private/ScriptContextComponent.cpp
new file mode 100644
--- /dev/null
+++ b/Engine/Plugins/Script/ScriptPlugin/Source/ScriptPlugin/Private/ScriptContextComponent.cpp
@@ -0,0 +1,89 @@
+// Copyright 1998-2014 Epic Games, Inc. All Rights Reserved. 
+#include "ScriptPluginPrivatePCH.h"
+
+//////////////////////////////////////////////////////////////////////////
+
+UScriptContextComponent::UScriptContextComponent(const FPostConstructInitializeProperties& PCIP)
+	: Super(PCIP)
+{
+	PrimaryComponentTick.bCanEverTick = true;
+	bTickInEditor = false;
+	bAutoActivate = true;
+	bWantsInitializeComponent = true;
+
+	Context = NULL;
+}
+
+void UScriptContextComponent::OnRegister()
+{
+	Super::OnRegister();
+
+	auto ContextOwner = GetOuter();
+	if (ContextOwner && !HasAnyFlags(RF_ClassDefaultObject) && !ContextOwner->HasAnyFlags(RF_ClassDefaultObject))
+	{
+		auto ScriptClass = UScriptBlueprintGeneratedClass::GetScriptGeneratedClass(ContextOwner->GetClass());
+		if (ScriptClass && GetWorld() && GetWorld()->WorldType != EWorldType::Editor)
+		{
+			Context = FScriptContextBase::CreateContext(ScriptClass->SourceCode, ScriptClass, ContextOwner);
+			if (!Context || !Context->CanTick())
+			{
+				bAutoActivate = false;
+				PrimaryComponentTick.bCanEverTick = false;
+			}
+		}
+	}
+}
+
+void UScriptContextComponent::InitializeComponent()
+{
+	Super::InitializeComponent();
+
+	if (Context)
+	{
+		auto ContextOwner = GetOuter();
+		auto ScriptClass = UScriptBlueprintGeneratedClass::GetScriptGeneratedClass(ContextOwner->GetClass());
+		check(ScriptClass);
+		Context->PushScriptPropertyValues(ScriptClass, ContextOwner);
+		Context->BeginPlay();
+		Context->FetchScriptPropertyValues(ScriptClass, ContextOwner);
+	}
+}
+
+void UScriptContextComponent::TickComponent(float DeltaTime, enum ELevelTick TickType, FActorComponentTickFunction *ThisTickFunction)
+{
+	Super::TickComponent(DeltaTime, TickType, ThisTickFunction);	
+	if (Context)
+	{
+		auto ContextOwner = GetOuter();
+		auto ScriptClass = UScriptBlueprintGeneratedClass::GetScriptGeneratedClass(ContextOwner->GetClass());
+		check(ScriptClass);
+		Context->PushScriptPropertyValues(ScriptClass, ContextOwner);
+		Context->Tick(DeltaTime);
+		Context->FetchScriptPropertyValues(ScriptClass, ContextOwner);
+	}
+};
+
+void UScriptContextComponent::BeginDestroy()
+{
+	if (Context)
+	{
+		Context->Destroy();
+		delete Context;
+		Context = NULL;
+	}
+
+	Super::BeginDestroy();
+}
+
+void UScriptContextComponent::CallScriptFunction(FString FunctionName)
+{	
+	if (Context)
+	{
+		auto ContextOwner = GetOuter();
+		auto ScriptClass = UScriptBlueprintGeneratedClass::GetScriptGeneratedClass(ContextOwner->GetClass());
+		check(ScriptClass);
+		Context->PushScriptPropertyValues(ScriptClass, ContextOwner);
+		Context->CallFunction(FunctionName);
+		Context->FetchScriptPropertyValues(ScriptClass, ContextOwner);
+	}
+}
diff --git a/Engine/Plugins/Script/ScriptPlugin/Source/ScriptPlugin/Private/ScriptFunction.cpp b/Engine/Plugins/Script/ScriptPlugin/Source/ScriptPlugin/Private/ScriptFunction.cpp
deleted file mode 100644
diff --git a/Engine/Plugins/Script/ScriptPlugin/Source/ScriptPlugin/Private/ScriptPluginPrivatePCH.h b/Engine/Plugins/Script/ScriptPlugin/Source/ScriptPlugin/Private/ScriptPluginPrivatePCH.h
--- a/Engine/Plugins/Script/ScriptPlugin/Source/ScriptPlugin/Private/ScriptPluginPrivatePCH.h
+++ b/Engine/Plugins/Script/ScriptPlugin/Source/ScriptPlugin/Private/ScriptPluginPrivatePCH.h
@@ -6,6 +6,8 @@
 #include "ModuleManager.h"
 #include "Engine.h"
 #include "ScriptComponent.h"
+#include "ScriptContext.h"
+#include "ScriptContextComponent.h"
 #include "ScriptTestActor.h"
 
 #if WITH_EDITOR
